[AI Toolsets]: Move UMG toolset tests under `AI.Toolsets`.

#Synced-CL 52596940

diff --git a/Engine/Plugins/Experimental/Toolsets/UMGToolSet/Source/UMGToolSet/Private/Tests/UMGToolSetTest.cpp b/Engine/Plugins/Experimental/Toolsets/UMGToolSet/Source/UMGToolSet/Private/Tests/UMGToolSetTest.cpp
--- a/Engine/Plugins/Experimental/Toolsets/UMGToolSet/Source/UMGToolSet/Private/Tests/UMGToolSetTest.cpp
+++ b/Engine/Plugins/Experimental/Toolsets/UMGToolSet/Source/UMGToolSet/Private/Tests/UMGToolSetTest.cpp
@@ -22,7 +22,7 @@
 // ============================================================================
 
 BEGIN_DEFINE_SPEC(FUMGToolSetTest_HappyPath,
-	"AI.UMGToolSet.HappyPath", UMGToolSetTest::Flags)
+	"AI.Toolsets.UMGToolSet.HappyPath", UMGToolSetTest::Flags)
 	UWidgetBlueprint* BP = nullptr;
 	int32 TestCounter = 0;
 END_DEFINE_SPEC(FUMGToolSetTest_HappyPath)
@@ -219,7 +219,7 @@ void FUMGToolSetTest_HappyPath::Define()
 // ============================================================================
 
 BEGIN_DEFINE_SPEC(FUMGToolSetTest_Errors,
-	"AI.UMGToolSet.Errors", UMGToolSetTest::Flags)
+	"AI.Toolsets.UMGToolSet.Errors", UMGToolSetTest::Flags)
 	UWidgetBlueprint* BP = nullptr;
 	int32 TestCounter = 0;
 END_DEFINE_SPEC(FUMGToolSetTest_Errors)
@@ -347,7 +347,7 @@ void FUMGToolSetTest_Errors::Define()
 // ============================================================================
 
 BEGIN_DEFINE_SPEC(FUMGToolSetTest_Resilience,
-	"AI.UMGToolSet.Resilience", UMGToolSetTest::Flags)
+	"AI.Toolsets.UMGToolSet.Resilience", UMGToolSetTest::Flags)
 	UWidgetBlueprint* BP = nullptr;
 	int32 TestCounter = 0;
 END_DEFINE_SPEC(FUMGToolSetTest_Resilience)
@@ -436,7 +436,7 @@ void FUMGToolSetTest_Resilience::Define()
 // ============================================================================
 
 BEGIN_DEFINE_SPEC(FUMGToolSetTest_BugFixes,
-	"AI.UMGToolSet.BugFixes", UMGToolSetTest::Flags)
+	"AI.Toolsets.UMGToolSet.BugFixes", UMGToolSetTest::Flags)
 	UWidgetBlueprint* BP = nullptr;
 	int32 TestCounter = 0;
 END_DEFINE_SPEC(FUMGToolSetTest_BugFixes)
@@ -711,7 +711,7 @@ void FUMGToolSetTest_BugFixes::Define()
 // ============================================================================
 
 BEGIN_DEFINE_SPEC(FUMGToolSetTest_BindWidget,
-	"AI.UMGToolSet.BindWidget", UMGToolSetTest::Flags)
+	"AI.Toolsets.UMGToolSet.BindWidget", UMGToolSetTest::Flags)
 	UWidgetBlueprint* BP = nullptr;
 	int32 TestCounter = 0;
 END_DEFINE_SPEC(FUMGToolSetTest_BindWidget)
@@ -780,7 +780,7 @@ void FUMGToolSetTest_BindWidget::Define()
 
 IMPLEMENT_SIMPLE_AUTOMATION_TEST(
 	FUMGToolSetTest_Registration,
-	"AI.UMGToolSet.Registration.ToolsetRegistered",
+	"AI.Toolsets.UMGToolSet.Registration.ToolsetRegistered",
 	UMGToolSetTest::Flags)
 
 bool FUMGToolSetTest_Registration::RunTest(const FString&)
@@ -795,7 +795,7 @@ bool FUMGToolSetTest_Registration::RunTest(const FString&)
 // ============================================================================
 
 BEGIN_DEFINE_SPEC(FUMGToolSetTest_CompileErrors,
-	"AI.UMGToolSet.CompileErrors", UMGToolSetTest::Flags)
+	"AI.Toolsets.UMGToolSet.CompileErrors", UMGToolSetTest::Flags)
 	int32 TestCounter = 0;
 END_DEFINE_SPEC(FUMGToolSetTest_CompileErrors)
 
